Don't crash when drawing a custom mesh or cable component if WireframeMaterial is NULL

diff --git a/Engine/Plugins/Runtime/CableComponent/Source/CableComponent/Private/CableComponent.cpp b/Engine/Plugins/Runtime/CableComponent/Source/CableComponent/Private/CableComponent.cpp
--- a/Engine/Plugins/Runtime/CableComponent/Source/CableComponent/Private/CableComponent.cpp
+++ b/Engine/Plugins/Runtime/CableComponent/Source/CableComponent/Private/CableComponent.cpp
@@ -243,10 +243,10 @@ public:
 	{
 		QUICK_SCOPE_CYCLE_COUNTER( STAT_CableSceneProxy_DrawDynamicElements );
 
-		const bool bWireframe = View->Family->EngineShowFlags.Wireframe;
+		const bool bWireframe = GEngine->AllowDebugViewModes() && View->Family->EngineShowFlags.Wireframe;
 
 		FColoredMaterialRenderProxy WireframeMaterialInstance(
-			GEngine->WireframeMaterial->GetRenderProxy(IsSelected()),
+			GEngine->WireframeMaterial ? GEngine->WireframeMaterial->GetRenderProxy(IsSelected()) : NULL,
 			FLinearColor(0, 0.5f, 1.f)
 			);
 
diff --git a/Engine/Plugins/Runtime/CustomMeshComponent/Source/CustomMeshComponent/Private/CustomMeshComponent.cpp b/Engine/Plugins/Runtime/CustomMeshComponent/Source/CustomMeshComponent/Private/CustomMeshComponent.cpp
--- a/Engine/Plugins/Runtime/CustomMeshComponent/Source/CustomMeshComponent/Private/CustomMeshComponent.cpp
+++ b/Engine/Plugins/Runtime/CustomMeshComponent/Source/CustomMeshComponent/Private/CustomMeshComponent.cpp
@@ -143,10 +143,10 @@ public:
 	{
 		QUICK_SCOPE_CYCLE_COUNTER( STAT_CustomMeshSceneProxy_DrawDynamicElements );
 
-		const bool bWireframe = View->Family->EngineShowFlags.Wireframe;
+		const bool bWireframe = GEngine->AllowDebugViewModes() && View->Family->EngineShowFlags.Wireframe;
 
 		FColoredMaterialRenderProxy WireframeMaterialInstance(
-			GEngine->WireframeMaterial->GetRenderProxy(IsSelected()),
+			GEngine->WireframeMaterial ? GEngine->WireframeMaterial->GetRenderProxy(IsSelected()) : NULL,
 			FLinearColor(0, 0.5f, 1.f)
 			);
 
