Fix compile error due to incorrect function names.

diff --git a/Engine/Plugins/Runtime/CustomMeshComponent/Source/CustomMeshComponent/Private/CustomMeshComponent.cpp b/Engine/Plugins/Runtime/CustomMeshComponent/Source/CustomMeshComponent/Private/CustomMeshComponent.cpp
--- a/Engine/Plugins/Runtime/CustomMeshComponent/Source/CustomMeshComponent/Private/CustomMeshComponent.cpp
+++ b/Engine/Plugins/Runtime/CustomMeshComponent/Source/CustomMeshComponent/Private/CustomMeshComponent.cpp
@@ -143,7 +143,7 @@ public:
 	{
 		QUICK_SCOPE_CYCLE_COUNTER( STAT_CustomMeshSceneProxy_DrawDynamicElements );
 
-		const bool bWireframe = GEngine->AllowDebugViewModes() && View->Family->EngineShowFlags.Wireframe;
+		const bool bWireframe = AllowDebugViewmodes() && View->Family->EngineShowFlags.Wireframe;
 
 		FColoredMaterialRenderProxy WireframeMaterialInstance(
 			GEngine->WireframeMaterial ? GEngine->WireframeMaterial->GetRenderProxy(IsSelected()) : NULL,
