GitHub PR-580: Add export API macro to UCustomMeshComponent

diff --git a/Engine/Plugins/Runtime/CustomMeshComponent/Source/CustomMeshComponent/Classes/CustomMeshComponent.h b/Engine/Plugins/Runtime/CustomMeshComponent/Source/CustomMeshComponent/Classes/CustomMeshComponent.h
--- a/Engine/Plugins/Runtime/CustomMeshComponent/Source/CustomMeshComponent/Classes/CustomMeshComponent.h
+++ b/Engine/Plugins/Runtime/CustomMeshComponent/Source/CustomMeshComponent/Classes/CustomMeshComponent.h
@@ -21,7 +21,7 @@ struct FCustomMeshTriangle
 
 /** Component that allows you to specify custom triangle mesh geometry */
 UCLASS(hidecategories=(Object,LOD, Physics, Collision), editinlinenew, meta=(BlueprintSpawnableComponent), ClassGroup=Rendering)
-class UCustomMeshComponent : public UMeshComponent
+class CUSTOMMESHCOMPONENT_API UCustomMeshComponent : public UMeshComponent
 {
 	GENERATED_UCLASS_BODY()
 
