Mark vertex members in FCustomMeshTriangle as BlueprintReadWrite
[https://github.com/EpicGames/UnrealEngine/pull/566] (stephenwhittle)

diff --git a/Engine/Plugins/Runtime/CustomMeshComponent/Source/CustomMeshComponent/Classes/CustomMeshComponent.h b/Engine/Plugins/Runtime/CustomMeshComponent/Source/CustomMeshComponent/Classes/CustomMeshComponent.h
--- a/Engine/Plugins/Runtime/CustomMeshComponent/Source/CustomMeshComponent/Classes/CustomMeshComponent.h
+++ b/Engine/Plugins/Runtime/CustomMeshComponent/Source/CustomMeshComponent/Classes/CustomMeshComponent.h
@@ -9,13 +9,13 @@ struct FCustomMeshTriangle
 {
 	GENERATED_USTRUCT_BODY()
 
-	UPROPERTY(EditAnywhere, Category=Triangle)
+	UPROPERTY(EditAnywhere, BlueprintReadWrite, Category=Triangle)
 	FVector Vertex0;
 
-	UPROPERTY(EditAnywhere, Category=Triangle)
+	UPROPERTY(EditAnywhere, BlueprintReadWrite, Category=Triangle)
 	FVector Vertex1;
 
-	UPROPERTY(EditAnywhere, Category=Triangle)
+	UPROPERTY(EditAnywhere, BlueprintReadWrite, Category=Triangle)
 	FVector Vertex2;
 };
 
