Merging

//depot/UE4-Orion/Engine/Plugins/Runtime/CustomMeshComponent/Source/...

diff --git a/Engine/Plugins/Runtime/CustomMeshComponent/Source/CustomMeshComponent/Classes/CustomMeshComponent.h b/Engine/Plugins/Runtime/CustomMeshComponent/Source/CustomMeshComponent/Classes/CustomMeshComponent.h
--- a/Engine/Plugins/Runtime/CustomMeshComponent/Source/CustomMeshComponent/Classes/CustomMeshComponent.h
+++ b/Engine/Plugins/Runtime/CustomMeshComponent/Source/CustomMeshComponent/Classes/CustomMeshComponent.h
@@ -29,6 +29,14 @@ class CUSTOMMESHCOMPONENT_API UCustomMeshComponent : public UMeshComponent
 	UFUNCTION(BlueprintCallable, Category="Components|CustomMesh")
 	bool SetCustomMeshTriangles(const TArray<FCustomMeshTriangle>& Triangles);
 
+	/** Add to the geometry to use on this triangle mesh.  This may cause an allocation.  Use SetCustomMeshTriangles() instead when possible to reduce allocations. */
+	UFUNCTION(BlueprintCallable, Category = "Components|CustomMesh")
+	void AddCustomMeshTriangles(const TArray<FCustomMeshTriangle>& Triangles);
+
+	/** Removes all geometry from this triangle mesh.  Does not deallocate memory, allowing new geometry to reuse the existing allocation. */
+	UFUNCTION(BlueprintCallable, Category = "Components|CustomMesh")
+	void ClearCustomMeshTriangles();
+
 private:
 
 	// Begin UPrimitiveComponent interface.
diff --git a/Engine/Plugins/Runtime/CustomMeshComponent/Source/CustomMeshComponent/Private/CustomMeshComponent.cpp b/Engine/Plugins/Runtime/CustomMeshComponent/Source/CustomMeshComponent/Private/CustomMeshComponent.cpp
--- a/Engine/Plugins/Runtime/CustomMeshComponent/Source/CustomMeshComponent/Private/CustomMeshComponent.cpp
+++ b/Engine/Plugins/Runtime/CustomMeshComponent/Source/CustomMeshComponent/Private/CustomMeshComponent.cpp
@@ -233,12 +233,29 @@ UCustomMeshComponent::UCustomMeshComponent( const FObjectInitializer& ObjectInit
 bool UCustomMeshComponent::SetCustomMeshTriangles(const TArray<FCustomMeshTriangle>& Triangles)
 {
 	CustomMeshTris = Triangles;
+
 	// Need to recreate scene proxy to send it over
 	MarkRenderStateDirty();
 
 	return true;
 }
 
+void UCustomMeshComponent::AddCustomMeshTriangles(const TArray<FCustomMeshTriangle>& Triangles)
+{
+	CustomMeshTris.Append(Triangles);
+
+	// Need to recreate scene proxy to send it over
+	MarkRenderStateDirty();
+}
+
+void  UCustomMeshComponent::ClearCustomMeshTriangles()
+{
+	CustomMeshTris.Reset();
+
+	// Need to recreate scene proxy to send it over
+	MarkRenderStateDirty();
+}
+
 
 FPrimitiveSceneProxy* UCustomMeshComponent::CreateSceneProxy()
 {
