FIX UE-16844 [CrashReport] Building lighting with a CustomMeshComponent in level causes crash

diff --git a/Engine/Plugins/Runtime/CustomMeshComponent/Source/CustomMeshComponent/Private/CustomMeshComponent.cpp b/Engine/Plugins/Runtime/CustomMeshComponent/Source/CustomMeshComponent/Private/CustomMeshComponent.cpp
--- a/Engine/Plugins/Runtime/CustomMeshComponent/Source/CustomMeshComponent/Private/CustomMeshComponent.cpp
+++ b/Engine/Plugins/Runtime/CustomMeshComponent/Source/CustomMeshComponent/Private/CustomMeshComponent.cpp
@@ -53,28 +53,40 @@ public:
 	FCustomMeshVertexFactory()
 	{}
 
+	/** Init function that should only be called on render thread. */
+	void Init_RenderThread(const FCustomMeshVertexBuffer* VertexBuffer)
+	{
+		check(IsInRenderingThread());
+
+		DataType NewData;
+		NewData.PositionComponent = STRUCTMEMBER_VERTEXSTREAMCOMPONENT(VertexBuffer, FDynamicMeshVertex, Position, VET_Float3);
+		NewData.TextureCoordinates.Add(
+			FVertexStreamComponent(VertexBuffer, STRUCT_OFFSET(FDynamicMeshVertex, TextureCoordinate), sizeof(FDynamicMeshVertex), VET_Float2)
+			);
+		NewData.TangentBasisComponents[0] = STRUCTMEMBER_VERTEXSTREAMCOMPONENT(VertexBuffer, FDynamicMeshVertex, TangentX, VET_PackedNormal);
+		NewData.TangentBasisComponents[1] = STRUCTMEMBER_VERTEXSTREAMCOMPONENT(VertexBuffer, FDynamicMeshVertex, TangentZ, VET_PackedNormal);
+		NewData.ColorComponent = STRUCTMEMBER_VERTEXSTREAMCOMPONENT(VertexBuffer, FDynamicMeshVertex, Color, VET_Color);
+
+		SetData(NewData);
+	}
 
 	/** Initialization */
 	void Init(const FCustomMeshVertexBuffer* VertexBuffer)
 	{
-		check(!IsInRenderingThread());
-
-		ENQUEUE_UNIQUE_RENDER_COMMAND_TWOPARAMETER(
-			InitCustomMeshVertexFactory,
-			FCustomMeshVertexFactory*,VertexFactory,this,
-			const FCustomMeshVertexBuffer*,VertexBuffer,VertexBuffer,
+		if (IsInRenderingThread())
+		{
+			Init_RenderThread(VertexBuffer);
+		}
+		else
 		{
-			// Initialize the vertex factory's stream components.
-			DataType NewData;
-			NewData.PositionComponent = STRUCTMEMBER_VERTEXSTREAMCOMPONENT(VertexBuffer,FDynamicMeshVertex,Position,VET_Float3);
-			NewData.TextureCoordinates.Add(
-				FVertexStreamComponent(VertexBuffer,STRUCT_OFFSET(FDynamicMeshVertex,TextureCoordinate),sizeof(FDynamicMeshVertex),VET_Float2)
-				);
-			NewData.TangentBasisComponents[0] = STRUCTMEMBER_VERTEXSTREAMCOMPONENT(VertexBuffer,FDynamicMeshVertex,TangentX,VET_PackedNormal);
-			NewData.TangentBasisComponents[1] = STRUCTMEMBER_VERTEXSTREAMCOMPONENT(VertexBuffer,FDynamicMeshVertex,TangentZ,VET_PackedNormal);
-			NewData.ColorComponent = STRUCTMEMBER_VERTEXSTREAMCOMPONENT(VertexBuffer, FDynamicMeshVertex, Color, VET_Color);
-			VertexFactory->SetData(NewData);
-		});
+			ENQUEUE_UNIQUE_RENDER_COMMAND_TWOPARAMETER(
+				InitCustomMeshVertexFactory,
+				FCustomMeshVertexFactory*, VertexFactory, this,
+				const FCustomMeshVertexBuffer*, VertexBuffer, VertexBuffer,
+				{
+				VertexFactory->Init_RenderThread(VertexBuffer);
+			});
+		}	
 	}
 };
 
