Check main pass and custom depth view relevance for custom and procedural mesh components. UE-18059

diff --git a/Engine/Plugins/Runtime/CustomMeshComponent/Source/CustomMeshComponent/Private/CustomMeshComponent.cpp b/Engine/Plugins/Runtime/CustomMeshComponent/Source/CustomMeshComponent/Private/CustomMeshComponent.cpp
--- a/Engine/Plugins/Runtime/CustomMeshComponent/Source/CustomMeshComponent/Private/CustomMeshComponent.cpp
+++ b/Engine/Plugins/Runtime/CustomMeshComponent/Source/CustomMeshComponent/Private/CustomMeshComponent.cpp
@@ -213,6 +213,8 @@ public:
 		Result.bDrawRelevance = IsShown(View);
 		Result.bShadowRelevance = IsShadowCast(View);
 		Result.bDynamicRelevance = true;
+		Result.bRenderInMainPass = ShouldRenderInMainPass();
+		Result.bRenderCustomDepth = ShouldRenderCustomDepth();
 		MaterialRelevance.SetPrimitiveViewRelevance(Result);
 		return Result;
 	}
diff --git a/Engine/Source/Runtime/ProceduralMeshComponent/Private/ProceduralMeshComponent.cpp b/Engine/Source/Runtime/ProceduralMeshComponent/Private/ProceduralMeshComponent.cpp
--- a/Engine/Source/Runtime/ProceduralMeshComponent/Private/ProceduralMeshComponent.cpp
+++ b/Engine/Source/Runtime/ProceduralMeshComponent/Private/ProceduralMeshComponent.cpp
@@ -273,6 +273,8 @@ public:
 		Result.bDrawRelevance = IsShown(View);
 		Result.bShadowRelevance = IsShadowCast(View);
 		Result.bDynamicRelevance = true;
+		Result.bRenderInMainPass = ShouldRenderInMainPass();
+		Result.bRenderCustomDepth = ShouldRenderCustomDepth();
 		MaterialRelevance.SetPrimitiveViewRelevance(Result);
 		return Result;
 	}
