Fix for uninitialized script struct members

diff --git a/Engine/Plugins/Runtime/CustomMeshComponent/Source/CustomMeshComponent/Classes/CustomMeshComponent.h b/Engine/Plugins/Runtime/CustomMeshComponent/Source/CustomMeshComponent/Classes/CustomMeshComponent.h
--- a/Engine/Plugins/Runtime/CustomMeshComponent/Source/CustomMeshComponent/Classes/CustomMeshComponent.h
+++ b/Engine/Plugins/Runtime/CustomMeshComponent/Source/CustomMeshComponent/Classes/CustomMeshComponent.h
@@ -15,13 +15,13 @@ struct FCustomMeshTriangle
 	GENERATED_USTRUCT_BODY()
 
 	UPROPERTY(EditAnywhere, BlueprintReadWrite, Category=Triangle)
-	FVector Vertex0;
+	FVector Vertex0 = FVector::ZeroVector;
 
 	UPROPERTY(EditAnywhere, BlueprintReadWrite, Category=Triangle)
-	FVector Vertex1;
+	FVector Vertex1 = FVector::ZeroVector;
 
 	UPROPERTY(EditAnywhere, BlueprintReadWrite, Category=Triangle)
-	FVector Vertex2;
+	FVector Vertex2 = FVector::ZeroVector;
 };
 
 /** Component that allows you to specify custom triangle mesh geometry */
diff --git a/Engine/Source/Editor/UnrealEd/Classes/MaterialEditor/DEditorFontParameterValue.h b/Engine/Source/Editor/UnrealEd/Classes/MaterialEditor/DEditorFontParameterValue.h
--- a/Engine/Source/Editor/UnrealEd/Classes/MaterialEditor/DEditorFontParameterValue.h
+++ b/Engine/Source/Editor/UnrealEd/Classes/MaterialEditor/DEditorFontParameterValue.h
@@ -17,7 +17,7 @@ struct UNREALED_API FDFontParameters
 	TObjectPtr<class UFont> FontValue = nullptr;
 
 	UPROPERTY(EditAnywhere, Category = DFontParameter)
-	int32 FontPage;
+	int32 FontPage = 0;
 };
 
 UCLASS(hidecategories = Object, collapsecategories)
diff --git a/Engine/Source/Runtime/Engine/Classes/Components/LightComponent.h b/Engine/Source/Runtime/Engine/Classes/Components/LightComponent.h
--- a/Engine/Source/Runtime/Engine/Classes/Components/LightComponent.h
+++ b/Engine/Source/Runtime/Engine/Classes/Components/LightComponent.h
@@ -497,7 +497,7 @@ struct FPrecomputedLightInstanceData : public FSceneComponentInstanceData
 	FGuid LightGuid;
 
 	UPROPERTY()
-	int32 PreviewShadowMapChannel;
+	int32 PreviewShadowMapChannel = 0;
 };
 
 
diff --git a/Engine/Source/Runtime/Engine/Classes/Components/SceneCaptureComponent.h b/Engine/Source/Runtime/Engine/Classes/Components/SceneCaptureComponent.h
--- a/Engine/Source/Runtime/Engine/Classes/Components/SceneCaptureComponent.h
+++ b/Engine/Source/Runtime/Engine/Classes/Components/SceneCaptureComponent.h
@@ -35,7 +35,7 @@ struct FEngineShowFlagsSetting
 	FString ShowFlagName;
 
 	UPROPERTY(EditAnywhere, BlueprintReadWrite, Category = SceneCapture)
-	bool Enabled;
+	bool Enabled = false;
 
 
 	bool operator == (const FEngineShowFlagsSetting& Other) const
diff --git a/Engine/Source/Runtime/Engine/Classes/Components/StaticMeshComponent.h b/Engine/Source/Runtime/Engine/Classes/Components/StaticMeshComponent.h
--- a/Engine/Source/Runtime/Engine/Classes/Components/StaticMeshComponent.h
+++ b/Engine/Source/Runtime/Engine/Classes/Components/StaticMeshComponent.h
@@ -749,7 +749,7 @@ struct FStaticMeshVertexColorLODData
 
 	/** Index of the LOD that this data came from */
 	UPROPERTY()
-	uint32 LODIndex;
+	uint32 LODIndex = 0;
 
 	/** Check whether this contains valid data */
 	bool IsValid() const
diff --git a/Engine/Source/Runtime/Engine/Classes/Engine/SkeletalMeshLODSettings.h b/Engine/Source/Runtime/Engine/Classes/Engine/SkeletalMeshLODSettings.h
--- a/Engine/Source/Runtime/Engine/Classes/Engine/SkeletalMeshLODSettings.h
+++ b/Engine/Source/Runtime/Engine/Classes/Engine/SkeletalMeshLODSettings.h
@@ -41,7 +41,7 @@ struct FBoneFilter
 	 * You can't remove joint without children removed, and you can't keep without your parents 
 	 */
 	UPROPERTY(EditAnywhere, Category = FBoneFilter)
-	bool bExcludeSelf;
+	bool bExcludeSelf = false;
 
 	/* Name of Bone Name */
 	UPROPERTY(EditAnywhere, Category = FBoneFilter)
diff --git a/Engine/Source/Runtime/Engine/Classes/Materials/MaterialExpressionCustom.h b/Engine/Source/Runtime/Engine/Classes/Materials/MaterialExpressionCustom.h
--- a/Engine/Source/Runtime/Engine/Classes/Materials/MaterialExpressionCustom.h
+++ b/Engine/Source/Runtime/Engine/Classes/Materials/MaterialExpressionCustom.h
@@ -43,7 +43,7 @@ struct FCustomOutput
 	FName OutputName;
 
 	UPROPERTY(EditAnywhere, Category = CustomOutput)
-	TEnumAsByte<enum ECustomMaterialOutputType> OutputType;
+	TEnumAsByte<enum ECustomMaterialOutputType> OutputType = ECustomMaterialOutputType::CMOT_Float1;
 };
 
 USTRUCT()
diff --git a/Engine/Source/Runtime/Engine/Public/ComputeFramework/ShaderParamTypeDefinition.h b/Engine/Source/Runtime/Engine/Public/ComputeFramework/ShaderParamTypeDefinition.h
--- a/Engine/Source/Runtime/Engine/Public/ComputeFramework/ShaderParamTypeDefinition.h
+++ b/Engine/Source/Runtime/Engine/Public/ComputeFramework/ShaderParamTypeDefinition.h
@@ -152,7 +152,7 @@ struct ENGINE_API FShaderValueType
 	EShaderFundamentalType Type;
 
 	UPROPERTY()
-	EShaderFundamentalDimensionType DimensionType;
+	EShaderFundamentalDimensionType DimensionType = EShaderFundamentalDimensionType::Scalar;
 
 	union
 	{
